[Motion Edit Rig]
- Move development to public

#rnx
#virtualized
#Synced-CL 52677197

diff --git a/Engine/Plugins/VirtualProduction/MotionEdit/Config/DefaultMotionEdit.ini b/Engine/Plugins/VirtualProduction/MotionEdit/Config/DefaultMotionEdit.ini
new file mode 100644
--- /dev/null
+++ b/Engine/Plugins/VirtualProduction/MotionEdit/Config/DefaultMotionEdit.ini
@@ -0,0 +1,20 @@
+﻿[CoreRedirects]
++StructRedirects=(OldName="/Script/MotionEditRig.MotionEditCharacterizationTemplate",NewName="/Script/MotionEditRig.BipedCharacterizationTemplate")
++StructRedirects=(OldName="/Script/MotionEditRig.MotionEditHandFootTemplate",NewName="/Script/MotionEditRig.BipedHandFootTemplate")
++StructRedirects=(OldName="/Script/MotionEditRig.MotionEditCharacterizationChain",NewName="/Script/MotionEditRig.BipedCharacterizationChain")
++ClassRedirects=(OldName="/Script/MotionEditRig.MotionEditCharacterization",NewName="/Script/MotionEditRig.BipedCharacterization")
++StructRedirects=(OldName="/Script/BipedTemplateRig.BipedExtremity",NewName="/Script/BipedTemplateRig.BipedHand")
+
++StructRedirects=(OldName="/Script/MotionEditRig.BipedCharacterizationTemplate",NewName="/Script/BipedTemplateRig.BipedCharacterizationTemplate")
++StructRedirects=(OldName="/Script/MotionEditRig.BipedHandFootTemplate",NewName="/Script/BipedTemplateRig.BipedHandFootTemplate")
++StructRedirects=(OldName="/Script/MotionEditRig.BipedCharacterizationChain",NewName="/Script/BipedTemplateRig.BipedCharacterizationChain")
++ClassRedirects=(OldName="/Script/MotionEditRig.BipedCharacterization",NewName="/Script/BipedTemplateRig.BipedCharacterization")
+
++PackageRedirects=(OldName="/MotionEditRig/...",NewName="/BipedTemplateRig/",MatchWildcard=true)
++PackageRedirects=(OldName="/BipedTemplateRig/...",NewName="/MotionEdit/",MatchWildcard=true)
++ObjectRedirects=(OldName="/BipedTemplateRig/Sandbox/MotionEditRigFunctionLibrary.MotionEditRigFunctionLibrary",NewName="/BipedTemplateRig/Sandbox/BipedTemplateFunctionLibrary.BipedTemplateFunctionLibrary")
++ObjectRedirects=(OldName="/BipedTemplateRig/Sandbox/BipedTemplateFunctionLibrary.MotionEditRigFunctionLibrary",NewName="/BipedTemplateRig/Sandbox/BipedTemplateFunctionLibrary.BipedTemplateFunctionLibrary")
++PropertyRedirects=(OldName="/Script/BipedTemplateRig.BipedFKChain.FBipedCharacterizationFKChain",NewName="/Script/BipedTemplateRig.BipedFKChain.BipedCharacterizationFKChain")
++PropertyRedirects=(OldName="/Script/BipedTemplateRig.BipedIKChain.FBipedCharacterizationIKChainStartBone",NewName="/Script/BipedTemplateRig.BipedIKChain.BipedCharacterizationIKChainStartBone")
++PropertyRedirects=(OldName="/Script/BipedTemplateRig.BipedIKChain.FBipedCharacterizationIKChainMidBone",NewName="/Script/BipedTemplateRig.BipedIKChain.BipedCharacterizationIKChainMidBone")
++PropertyRedirects=(OldName="/Script/BipedTemplateRig.BipedIKChain.FBipedCharacterizationIKChainEndBone",NewName="/Script/BipedTemplateRig.BipedIKChain.BipedCharacterizationIKChainEndBone")
\ No newline at end of file
diff --git a/Engine/Plugins/VirtualProduction/MotionEdit/MotionEdit.uplugin b/Engine/Plugins/VirtualProduction/MotionEdit/MotionEdit.uplugin
new file mode 100644
--- /dev/null
+++ b/Engine/Plugins/VirtualProduction/MotionEdit/MotionEdit.uplugin
@@ -0,0 +1,38 @@
+{
+	"FileVersion": 3,
+	"Version": 1,
+	"VersionName": "0.1",
+	"FriendlyName": "Motion Edit",
+	"Description": "Experimental Plugin for creating generic motion edit rigs and editing anim sequences",
+	"Category": "Animation",
+	"CreatedBy": "Epic Games",
+	"CreatedByURL": "",
+	"DocsURL": "",
+	"MarketplaceURL": "",
+	"SupportURL": "",
+	"CanContainContent": true,
+	"Installed": false,
+	"Modules": [
+		{
+			"Name": "BipedTemplateRig",
+			"Type": "Editor",
+			"LoadingPhase": "Default"
+		}
+	],
+	"Plugins": [
+		{
+			"Name": "CommonUI",
+			"Enabled": true
+		},
+		{
+			"Name": "PerformanceCaptureWorkflow",
+			"Enabled": true
+		},
+		{
+			"Name": "ControlRig",
+			"Enabled": true
+		}
+	],
+	"IsBetaVersion": false,
+	"IsExperimentalVersion": true
+}
\ No newline at end of file
diff --git a/Engine/Plugins/VirtualProduction/MotionEdit/Source/BipedTemplateRig/BipedTemplateRig.Build.cs b/Engine/Plugins/VirtualProduction/MotionEdit/Source/BipedTemplateRig/BipedTemplateRig.Build.cs
new file mode 100644
--- /dev/null
+++ b/Engine/Plugins/VirtualProduction/MotionEdit/Source/BipedTemplateRig/BipedTemplateRig.Build.cs
@@ -0,0 +1,49 @@
+﻿// Copyright Epic Games, Inc. All Rights Reserved.
+
+using UnrealBuildTool;
+
+public class BipedTemplateRig : ModuleRules
+{
+	public BipedTemplateRig(ReadOnlyTargetRules Target) : base(Target)
+	{
+		PCHUsage = ModuleRules.PCHUsageMode.UseExplicitOrSharedPCHs;
+		
+		PublicIncludePaths.AddRange(
+			new string[] {
+			}
+			);
+				
+		
+		PrivateIncludePaths.AddRange(
+			new string[] {
+			}
+			);
+			
+		
+		PublicDependencyModuleNames.AddRange(
+			new string[]
+			{
+				"Core",
+			}
+			);
+			
+		
+		PrivateDependencyModuleNames.AddRange(
+			new string[]
+			{
+				"CoreUObject",
+				"Engine",
+				"Slate",
+				"SlateCore",
+				"UnrealEd"
+			}
+			);
+		
+		
+		DynamicallyLoadedModuleNames.AddRange(
+			new string[]
+			{
+			}
+			);
+	}
+}
diff --git a/Engine/Plugins/VirtualProduction/MotionEdit/Source/BipedTemplateRig/Private/BipedRigCharacterization.cpp b/Engine/Plugins/VirtualProduction/MotionEdit/Source/BipedTemplateRig/Private/BipedRigCharacterization.cpp
new file mode 100644
--- /dev/null
+++ b/Engine/Plugins/VirtualProduction/MotionEdit/Source/BipedTemplateRig/Private/BipedRigCharacterization.cpp
@@ -0,0 +1,25 @@
+﻿// Copyright Epic Games, Inc. All Rights Reserved.
+
+
+#include "BipedRigCharacterization.h"
+
+FBipedHand::FBipedHand()
+{
+	Thumb.BipedCharacterizationFKChain.Init(NAME_None, 1);
+	// Initialize a temp copy of FK chain to then initialize the fingers array 
+	FBipedFKChain FKChainTemp;
+	FKChainTemp.BipedCharacterizationFKChain.Init(NAME_None, 2);
+	
+	Digits.Init(FKChainTemp,4);
+}
+
+FBipedTemplate::FBipedTemplate()
+{
+	Root.BipedCharacterizationFKChain.Init(NAME_None, 2);
+	Spine.BipedCharacterizationFKChain.Init(NAME_None, 2);
+	Neck.BipedCharacterizationFKChain.Init(NAME_None, 2);
+	LeftShoulder.BipedCharacterizationFKChain.Init(NAME_None, 1);
+	RightShoulder.BipedCharacterizationFKChain.Init(NAME_None, 1);
+	LeftFoot.BipedCharacterizationFKChain.Init(NAME_None, 1);
+	RightFoot.BipedCharacterizationFKChain.Init(NAME_None, 1);
+}
diff --git a/Engine/Plugins/VirtualProduction/MotionEdit/Source/BipedTemplateRig/Private/BipedRigCharacterization.h b/Engine/Plugins/VirtualProduction/MotionEdit/Source/BipedTemplateRig/Private/BipedRigCharacterization.h
new file mode 100644
--- /dev/null
+++ b/Engine/Plugins/VirtualProduction/MotionEdit/Source/BipedTemplateRig/Private/BipedRigCharacterization.h
@@ -0,0 +1,191 @@
+﻿// Copyright Epic Games, Inc. All Rights Reserved.
+
+#pragma once
+
+#include "CoreMinimal.h"
+#include "Engine/DataAsset.h"
+#include "Factories/Factory.h"
+#include "BipedRigCharacterization.generated.h"
+
+class USkeleton;
+
+USTRUCT(BlueprintType)
+struct FBipedFKChain
+{
+	GENERATED_BODY()
+	public:
+	/**
+	 * Struct for FK only chain - ie the spine or neck. 
+	 */
+	UPROPERTY(BlueprintReadWrite, EditAnywhere, Category="Biped Characterization", meta=(DisplayName="FK Chain"))
+	TArray<FName> BipedCharacterizationFKChain;
+};
+
+USTRUCT(BlueprintType)
+struct FBipedIKChain
+{
+	GENERATED_BODY()
+public:
+	/**
+	 * This is the first bone in an IK chain - typically the upper arm or upper leg
+	 */
+	UPROPERTY(BlueprintReadWrite, EditAnywhere, Category="Biped Characterization", meta=(DisplayName="Start Bone"))
+	FName BipedCharacterizationIKChainStartBone;
+
+	/**
+	 * The mid-bone in the chain - typically the elbow or knee. Used to determine the IK chain's pole
+	 * vector target
+	 */
+	UPROPERTY(BlueprintReadWrite, EditAnywhere, Category="Biped Characterization", meta=(DisplayName="Mid Bone"))
+	FName BipedCharacterizationIKChainMidBone;
+
+	/**
+	 * The last bone in the IK chain - typically the wrist/hand or ankle
+	 */
+	UPROPERTY(BlueprintReadWrite, EditAnywhere, Category="Biped Characterization", meta=(DisplayName="End Bone"))
+	FName BipedCharacterizationIKChainEndBone;
+};
+
+USTRUCT(BlueprintType)
+struct FBipedHand
+{
+	GENERATED_BODY()
+public:
+	
+	FBipedHand();
+	
+	/**
+	 * HandBase bone
+	 */
+	UPROPERTY(BlueprintReadWrite, EditAnywhere, Category="Biped Characterization", meta=(DisplayName="Hand"))
+	FName HandBase;
+	
+	/**
+	 * Thumb chain
+	 */
+	UPROPERTY(BlueprintReadWrite, EditAnywhere, Category="Biped Characterization", meta=(DisplayName="Thumb"))
+	FBipedFKChain Thumb;
+	
+	/**
+	 * An array of finger chains. Assume index[0] is the index finger, then we  up from index to pinky and beyond if required
+	 */
+	UPROPERTY(BlueprintReadWrite, EditAnywhere, Category="Biped Characterization", meta=(DisplayName="Fingers"))
+	TArray<FBipedFKChain> Digits;
+};
+
+/**
+ * Characterization struct v02
+ */
+
+USTRUCT(BlueprintType)
+struct 	FBipedTemplate
+
+{
+	GENERATED_BODY()
+	
+	FBipedTemplate();
+public:
+	
+	/**
+	 * Root chain
+	 */
+	UPROPERTY(BlueprintReadWrite, EditAnywhere, Category="Biped Characterization", meta=(DisplayName="Root/Core"))
+	FBipedFKChain Root;
+	
+	/**
+	 * Spine chain
+	 */
+	UPROPERTY(BlueprintReadWrite, EditAnywhere, Category="Biped Characterization", meta=(DisplayName="Spine"))
+	FBipedFKChain Spine;
+		
+	/**
+	 * Neck and Head bones
+	 */
+	UPROPERTY(BlueprintReadWrite, EditAnywhere, Category="Biped Characterization", meta=(DisplayName="Neck"))
+	FBipedFKChain Neck;
+	
+	/**
+	 * Right Shoulder
+	 */
+	UPROPERTY(BlueprintReadWrite, EditAnywhere, Category="Biped Characterization", meta=(DisplayName="Right Shoulder"))
+	FBipedFKChain RightShoulder;
+	
+	/**
+	 * Left Shoulder
+	 */
+	UPROPERTY(BlueprintReadWrite, EditAnywhere, Category="Biped Characterization", meta=(DisplayName="Left Shoulder"))
+	FBipedFKChain LeftShoulder;
+				
+	/**
+	 * Right Arm
+	 */
+	UPROPERTY(BlueprintReadWrite, EditAnywhere, Category="Biped Characterization", meta=(DisplayName="Right Arm"))
+	FBipedIKChain RightArm;
+		
+	/**
+	 * Left Arm
+	 */
+	UPROPERTY(BlueprintReadWrite, EditAnywhere, Category="Biped Characterization", meta=(DisplayName="Left Arm"))
+	FBipedIKChain LeftArm;
+						
+	/**
+	 * Right Hand
+	 */
+	UPROPERTY(BlueprintReadWrite, EditAnywhere, Category="Biped Characterization", meta=(DisplayName="Right Hand"))
+	FBipedHand RightHand;
+	
+	/**
+	 * Left Hand
+	 */
+	UPROPERTY(BlueprintReadWrite, EditAnywhere, Category="Biped Characterization", meta=(DisplayName="Left Hand"))
+	FBipedHand LeftHand;
+	
+	/**
+	 * Right Leg
+	 */
+	UPROPERTY(BlueprintReadWrite, EditAnywhere, Category="Biped Characterization", meta=(DisplayName="Right Leg"))
+	FBipedIKChain RightLeg;
+	
+	/**
+	 * Left Leg
+	 */
+	UPROPERTY(BlueprintReadWrite, EditAnywhere, Category="Biped Characterization", meta=(DisplayName="Left Leg"))
+	FBipedIKChain LeftLeg;
+	
+	/**
+	 * Right Foot
+	 */
+	UPROPERTY(BlueprintReadWrite, EditAnywhere, Category="Biped Characterization", meta=(DisplayName="Right Foot"))
+	FBipedFKChain RightFoot;
+	
+	/**
+	 * Left Foot
+	 */
+	UPROPERTY(BlueprintReadWrite, EditAnywhere, Category="Biped Characterization", meta=(DisplayName="Left Foot"))
+	FBipedFKChain LeftFoot;
+	
+	/**
+	 * Right Toes. First entry should the big toes and then we count up from there.
+	 */
+	UPROPERTY(BlueprintReadWrite, EditAnywhere, Category="Biped Characterization", meta=(DisplayName="Right Toes"))
+	TArray<FBipedFKChain> RightFootToes;
+	
+	/**
+	 * Left Toes. First entry should the big toes and then we count up from there. 
+	 */
+	UPROPERTY(BlueprintReadWrite, EditAnywhere, Category="Biped Characterization", meta=(DisplayName="Left Toes"))
+	TArray<FBipedFKChain> LeftFootToes;
+};
+
+UCLASS(Blueprintable, BlueprintType)
+
+class UBipedCharacterization : public UDataAsset
+
+{
+	GENERATED_BODY()
+
+public:
+	
+	UPROPERTY(BlueprintReadWrite, EditAnywhere, Category = "Biped Template", meta=(DisplayName="Characterization Template"))
+	FBipedTemplate Characterization;
+};
diff --git a/Engine/Plugins/VirtualProduction/MotionEdit/Source/BipedTemplateRig/Private/BipedTemplateRig.cpp b/Engine/Plugins/VirtualProduction/MotionEdit/Source/BipedTemplateRig/Private/BipedTemplateRig.cpp
new file mode 100644
--- /dev/null
+++ b/Engine/Plugins/VirtualProduction/MotionEdit/Source/BipedTemplateRig/Private/BipedTemplateRig.cpp
@@ -0,0 +1,20 @@
+﻿// Copyright Epic Games, Inc. All Rights Reserved.
+
+#include "BipedTemplateRig.h"
+
+#define LOCTEXT_NAMESPACE "FBipedTemplateRigModule"
+
+void FBipedTemplateRigModule::StartupModule()
+{
+	// This code will execute after your module is loaded into memory; the exact timing is specified in the .uplugin file per-module
+}
+
+void FBipedTemplateRigModule::ShutdownModule()
+{
+	// This function may be called during shutdown to clean up your module.  For modules that support dynamic reloading,
+	// we call this function before unloading the module.
+}
+
+#undef LOCTEXT_NAMESPACE
+	
+IMPLEMENT_MODULE(FBipedTemplateRigModule, BipedTemplateRig)
\ No newline at end of file
diff --git a/Engine/Plugins/VirtualProduction/MotionEdit/Source/BipedTemplateRig/Public/BipedTemplateRig.h b/Engine/Plugins/VirtualProduction/MotionEdit/Source/BipedTemplateRig/Public/BipedTemplateRig.h
new file mode 100644
--- /dev/null
+++ b/Engine/Plugins/VirtualProduction/MotionEdit/Source/BipedTemplateRig/Public/BipedTemplateRig.h
@@ -0,0 +1,14 @@
+﻿// Copyright Epic Games, Inc. All Rights Reserved.
+
+#pragma once
+
+#include "Modules/ModuleManager.h"
+
+class FBipedTemplateRigModule : public IModuleInterface
+{
+public:
+
+	/** IModuleInterface implementation */
+	virtual void StartupModule() override;
+	virtual void ShutdownModule() override;
+};
